PR #12093: add error detection for SQLiteSupport

diff --git a/Engine/Plugins/Runtime/Database/DatabaseSupport/Source/DatabaseSupport/Public/Database.h b/Engine/Plugins/Runtime/Database/DatabaseSupport/Source/DatabaseSupport/Public/Database.h
--- a/Engine/Plugins/Runtime/Database/DatabaseSupport/Source/DatabaseSupport/Public/Database.h
+++ b/Engine/Plugins/Runtime/Database/DatabaseSupport/Source/DatabaseSupport/Public/Database.h
@@ -65,6 +65,28 @@ protected:
 
 public:
 
+	/**
+	 * @return True if the record set has an error set.
+	 */
+	virtual bool HasError() const
+	{
+		return false;
+	}
+	/**
+	 * @return The raw error code that was set, if any (@see HasError).
+	 */
+	virtual int32 GetErrorCode() const
+	{
+		return 0;
+	}
+	/**
+	 * @return The error message that was set, if any (@see HasError).
+	 */
+	virtual FString GetErrorMessage() const
+	{
+		return FString();
+	}
+
 	/** 
 	 *   Returns a count of the number of records in the record set
 	 */
@@ -236,4 +258,21 @@ public:
 		RecordSet = nullptr;
 		return false;
 	}
+
+	/**
+	 * Executes the passed in command on the database. The caller is responsible for deleting
+	 * the created RecordSet.
+	 *
+	 * @param CommandString		Command to execute
+	 * @param RecordSet			Reference to recordset pointer that is going to hold result
+	 *
+	 * @return true if execution was successful, false otherwise
+	 */
+	inline bool Execute(const TCHAR* CommandString, TUniquePtr< FDataBaseRecordSet >& RecordSet)
+	{
+		FDataBaseRecordSet* RecordSetRaw = nullptr;
+		bool bRet = Execute(CommandString, RecordSetRaw);
+		RecordSet.Reset(RecordSetRaw);
+		return bRet;
+	}
 };
diff --git a/Engine/Plugins/Runtime/Database/SQLiteCore/Source/SQLiteCore/Private/SQLiteDatabase.cpp b/Engine/Plugins/Runtime/Database/SQLiteCore/Source/SQLiteCore/Private/SQLiteDatabase.cpp
--- a/Engine/Plugins/Runtime/Database/SQLiteCore/Source/SQLiteCore/Private/SQLiteDatabase.cpp
+++ b/Engine/Plugins/Runtime/Database/SQLiteCore/Source/SQLiteCore/Private/SQLiteDatabase.cpp
@@ -194,6 +194,15 @@ FSQLitePreparedStatement FSQLiteDatabase::PrepareStatement(const TCHAR* InStatem
 		: FSQLitePreparedStatement();
 }
 
+FString FSQLiteDatabase::GetErrorForCode(const int32 ResultCode)
+{
+	if (const char* ErrorStr = sqlite3_errstr(ResultCode))
+	{
+		return UTF8_TO_TCHAR(ErrorStr);
+	}
+	return FString();
+}
+
 FString FSQLiteDatabase::GetLastError() const
 {
 	const char* ErrorStr = Database ? sqlite3_errmsg(Database) : nullptr;
diff --git a/Engine/Plugins/Runtime/Database/SQLiteCore/Source/SQLiteCore/Private/SQLitePreparedStatement.cpp b/Engine/Plugins/Runtime/Database/SQLiteCore/Source/SQLiteCore/Private/SQLitePreparedStatement.cpp
--- a/Engine/Plugins/Runtime/Database/SQLiteCore/Source/SQLiteCore/Private/SQLitePreparedStatement.cpp
+++ b/Engine/Plugins/Runtime/Database/SQLiteCore/Source/SQLiteCore/Private/SQLitePreparedStatement.cpp
@@ -420,10 +420,14 @@ int64 FSQLitePreparedStatement::Execute(TFunctionRef<ESQLitePreparedStatementExe
 		: RowCount;
 }
 
-ESQLitePreparedStatementStepResult FSQLitePreparedStatement::Step()
+ESQLitePreparedStatementStepResult FSQLitePreparedStatement::Step(int32* OutResultCode)
 {
 	if (!Statement)
 	{
+		if (OutResultCode)
+		{
+			*OutResultCode = SQLITE_ERROR;
+		}
 		return ESQLitePreparedStatementStepResult::Error;
 	}
 
@@ -437,6 +441,10 @@ ESQLitePreparedStatementStepResult FSQLitePreparedStatement::Step()
 	case SQLITE_BUSY:
 		return ESQLitePreparedStatementStepResult::Busy;
 	default:
+		if (OutResultCode)
+		{
+			*OutResultCode = Result;
+		}
 		return ESQLitePreparedStatementStepResult::Error;
 	}
 }
diff --git a/Engine/Plugins/Runtime/Database/SQLiteCore/Source/SQLiteCore/Public/SQLiteDatabase.h b/Engine/Plugins/Runtime/Database/SQLiteCore/Source/SQLiteCore/Public/SQLiteDatabase.h
--- a/Engine/Plugins/Runtime/Database/SQLiteCore/Source/SQLiteCore/Public/SQLiteDatabase.h
+++ b/Engine/Plugins/Runtime/Database/SQLiteCore/Source/SQLiteCore/Public/SQLiteDatabase.h
@@ -122,6 +122,11 @@ public:
 			: T();
 	}
 
+	/**
+	 * Get the error string associated with the given result code, if any.
+	 */
+	static UE_API FString GetErrorForCode(const int32 ResultCode);
+
 	/**
 	 * Get the last error reported by this database.
 	 */
diff --git a/Engine/Plugins/Runtime/Database/SQLiteCore/Source/SQLiteCore/Public/SQLitePreparedStatement.h b/Engine/Plugins/Runtime/Database/SQLiteCore/Source/SQLiteCore/Public/SQLitePreparedStatement.h
--- a/Engine/Plugins/Runtime/Database/SQLiteCore/Source/SQLiteCore/Public/SQLitePreparedStatement.h
+++ b/Engine/Plugins/Runtime/Database/SQLiteCore/Source/SQLiteCore/Public/SQLitePreparedStatement.h
@@ -213,9 +213,10 @@ public:
 
 	/**
 	 * Step the SQLite prepared statement to try and move on to the next result from the statement.
+	 * @param OutResultCode Optional value to store the low-level result code of the step (@see FSQLiteDatabase::GetErrorForCode).
 	 * @note See FSQLiteDatabase::Execute for a simple example of stepping a statement.
 	 */
-	UE_API ESQLitePreparedStatementStepResult Step();
+	UE_API ESQLitePreparedStatementStepResult Step(int32* OutResultCode = nullptr);
 
 	/**
 	 * Get the index of a column from its name.
diff --git a/Engine/Plugins/Runtime/Database/SQLiteSupport/Source/SQLiteSupport/Private/SQLiteResultSet.cpp b/Engine/Plugins/Runtime/Database/SQLiteSupport/Source/SQLiteSupport/Private/SQLiteResultSet.cpp
--- a/Engine/Plugins/Runtime/Database/SQLiteSupport/Source/SQLiteSupport/Private/SQLiteResultSet.cpp
+++ b/Engine/Plugins/Runtime/Database/SQLiteSupport/Source/SQLiteSupport/Private/SQLiteResultSet.cpp
@@ -2,6 +2,7 @@
 
 #include "SQLiteResultSet.h"
 #include "CoreGlobals.h"
+#include "SQLiteDatabase.h"
 #include "SQLiteTypes.h"
 
 FSQLiteResultSet::FSQLiteResultSet(FSQLitePreparedStatement&& InPreparedStatement)
@@ -9,7 +10,7 @@ FSQLiteResultSet::FSQLiteResultSet(FSQLitePreparedStatement&& InPreparedStatemen
 {
 	check(PreparedStatement.IsValid());
 
-	StepResult = PreparedStatement.Step();
+	StepResult = PreparedStatement.Step(&StepResultCode);
 	if (StepResult == ESQLitePreparedStatementStepResult::Row)
 	{
 		const TArray<FString>& ColumnNames = PreparedStatement.GetColumnNames();
@@ -108,13 +109,28 @@ bool FSQLiteResultSet::IsAtEnd() const
 	return StepResult == ESQLitePreparedStatementStepResult::Done;
 }
 
+bool FSQLiteResultSet::HasError() const
+{
+	return StepResult == ESQLitePreparedStatementStepResult::Error;
+}
+
+int32 FSQLiteResultSet::GetErrorCode() const
+{
+	return StepResultCode;
+}
+
+FString FSQLiteResultSet::GetErrorMessage() const
+{
+	return FString::Printf(TEXT("%s(%d)"), *FSQLiteDatabase::GetErrorForCode(StepResultCode), StepResultCode);
+}
+
 void FSQLiteResultSet::MoveToNext()
 {
-	StepResult = PreparedStatement.Step();
+	StepResult = PreparedStatement.Step(&StepResultCode);
 }
 
 void FSQLiteResultSet::MoveToFirst()
 {
 	PreparedStatement.Reset();
-	StepResult = PreparedStatement.Step();
+	StepResult = PreparedStatement.Step(&StepResultCode);
 }
diff --git a/Engine/Plugins/Runtime/Database/SQLiteSupport/Source/SQLiteSupport/Public/SQLiteResultSet.h b/Engine/Plugins/Runtime/Database/SQLiteSupport/Source/SQLiteSupport/Public/SQLiteResultSet.h
--- a/Engine/Plugins/Runtime/Database/SQLiteSupport/Source/SQLiteSupport/Public/SQLiteResultSet.h
+++ b/Engine/Plugins/Runtime/Database/SQLiteSupport/Source/SQLiteSupport/Public/SQLiteResultSet.h
@@ -18,6 +18,9 @@ protected:
 	virtual bool IsAtEnd() const override;
 
 public:
+	virtual bool HasError() const override;
+	virtual int32 GetErrorCode() const override;
+	virtual FString GetErrorMessage() const override;
 	virtual int32 GetRecordCount() const override;
 	virtual FString GetString(const TCHAR* Column) const override;
 	virtual int32 GetInt(const TCHAR* Column) const override;
@@ -32,5 +35,6 @@ private:
 	FSQLitePreparedStatement PreparedStatement;
 	TArray<FDatabaseColumnInfo> ColumnInfos;
 	int32 NumberOfRecords = 0;
+	int32 StepResultCode = 0;
 	ESQLitePreparedStatementStepResult StepResult = ESQLitePreparedStatementStepResult::Done;
 };
