Clarified that the output param of Step was only filled for error cases, and not when a row is returned

#rnx

diff --git a/Engine/Plugins/Runtime/Database/SQLiteCore/Source/SQLiteCore/Private/SQLitePreparedStatement.cpp b/Engine/Plugins/Runtime/Database/SQLiteCore/Source/SQLiteCore/Private/SQLitePreparedStatement.cpp
--- a/Engine/Plugins/Runtime/Database/SQLiteCore/Source/SQLiteCore/Private/SQLitePreparedStatement.cpp
+++ b/Engine/Plugins/Runtime/Database/SQLiteCore/Source/SQLiteCore/Private/SQLitePreparedStatement.cpp
@@ -420,13 +420,18 @@ int64 FSQLitePreparedStatement::Execute(TFunctionRef<ESQLitePreparedStatementExe
 		: RowCount;
 }
 
-ESQLitePreparedStatementStepResult FSQLitePreparedStatement::Step(int32* OutResultCode)
+ESQLitePreparedStatementStepResult FSQLitePreparedStatement::Step(int32* OutErrorCode)
 {
+	if (OutErrorCode)
+	{
+		*OutErrorCode = SQLITE_OK;
+	}
+
 	if (!Statement)
 	{
-		if (OutResultCode)
+		if (OutErrorCode)
 		{
-			*OutResultCode = SQLITE_ERROR;
+			*OutErrorCode = SQLITE_ERROR;
 		}
 		return ESQLitePreparedStatementStepResult::Error;
 	}
@@ -441,9 +446,9 @@ ESQLitePreparedStatementStepResult FSQLitePreparedStatement::Step(int32* OutResu
 	case SQLITE_BUSY:
 		return ESQLitePreparedStatementStepResult::Busy;
 	default:
-		if (OutResultCode)
+		if (OutErrorCode)
 		{
-			*OutResultCode = Result;
+			*OutErrorCode = Result;
 		}
 		return ESQLitePreparedStatementStepResult::Error;
 	}
diff --git a/Engine/Plugins/Runtime/Database/SQLiteCore/Source/SQLiteCore/Public/SQLitePreparedStatement.h b/Engine/Plugins/Runtime/Database/SQLiteCore/Source/SQLiteCore/Public/SQLitePreparedStatement.h
--- a/Engine/Plugins/Runtime/Database/SQLiteCore/Source/SQLiteCore/Public/SQLitePreparedStatement.h
+++ b/Engine/Plugins/Runtime/Database/SQLiteCore/Source/SQLiteCore/Public/SQLitePreparedStatement.h
@@ -213,10 +213,10 @@ public:
 
 	/**
 	 * Step the SQLite prepared statement to try and move on to the next result from the statement.
-	 * @param OutResultCode Optional value to store the low-level result code of the step (@see FSQLiteDatabase::GetErrorForCode).
+	 * @param OutResultCode Optional value to store the low-level error code of the step (@see FSQLiteDatabase::GetErrorForCode).
 	 * @note See FSQLiteDatabase::Execute for a simple example of stepping a statement.
 	 */
-	UE_API ESQLitePreparedStatementStepResult Step(int32* OutResultCode = nullptr);
+	UE_API ESQLitePreparedStatementStepResult Step(int32* OutErrorCode = nullptr);
 
 	/**
 	 * Get the index of a column from its name.
diff --git a/Engine/Plugins/Runtime/Database/SQLiteSupport/Source/SQLiteSupport/Private/SQLiteResultSet.cpp b/Engine/Plugins/Runtime/Database/SQLiteSupport/Source/SQLiteSupport/Private/SQLiteResultSet.cpp
--- a/Engine/Plugins/Runtime/Database/SQLiteSupport/Source/SQLiteSupport/Private/SQLiteResultSet.cpp
+++ b/Engine/Plugins/Runtime/Database/SQLiteSupport/Source/SQLiteSupport/Private/SQLiteResultSet.cpp
@@ -10,7 +10,7 @@ FSQLiteResultSet::FSQLiteResultSet(FSQLitePreparedStatement&& InPreparedStatemen
 {
 	check(PreparedStatement.IsValid());
 
-	StepResult = PreparedStatement.Step(&StepResultCode);
+	StepResult = PreparedStatement.Step(&StepResultErrorCode);
 	if (StepResult == ESQLitePreparedStatementStepResult::Row)
 	{
 		const TArray<FString>& ColumnNames = PreparedStatement.GetColumnNames();
@@ -116,21 +116,21 @@ bool FSQLiteResultSet::HasError() const
 
 int32 FSQLiteResultSet::GetErrorCode() const
 {
-	return StepResultCode;
+	return StepResultErrorCode;
 }
 
 FString FSQLiteResultSet::GetErrorMessage() const
 {
-	return FString::Printf(TEXT("%s(%d)"), *FSQLiteDatabase::GetErrorForCode(StepResultCode), StepResultCode);
+	return FString::Printf(TEXT("%s(%d)"), *FSQLiteDatabase::GetErrorForCode(StepResultErrorCode), StepResultErrorCode);
 }
 
 void FSQLiteResultSet::MoveToNext()
 {
-	StepResult = PreparedStatement.Step(&StepResultCode);
+	StepResult = PreparedStatement.Step(&StepResultErrorCode);
 }
 
 void FSQLiteResultSet::MoveToFirst()
 {
 	PreparedStatement.Reset();
-	StepResult = PreparedStatement.Step(&StepResultCode);
+	StepResult = PreparedStatement.Step(&StepResultErrorCode);
 }
diff --git a/Engine/Plugins/Runtime/Database/SQLiteSupport/Source/SQLiteSupport/Public/SQLiteResultSet.h b/Engine/Plugins/Runtime/Database/SQLiteSupport/Source/SQLiteSupport/Public/SQLiteResultSet.h
--- a/Engine/Plugins/Runtime/Database/SQLiteSupport/Source/SQLiteSupport/Public/SQLiteResultSet.h
+++ b/Engine/Plugins/Runtime/Database/SQLiteSupport/Source/SQLiteSupport/Public/SQLiteResultSet.h
@@ -35,6 +35,6 @@ private:
 	FSQLitePreparedStatement PreparedStatement;
 	TArray<FDatabaseColumnInfo> ColumnInfos;
 	int32 NumberOfRecords = 0;
-	int32 StepResultCode = 0;
+	int32 StepResultErrorCode = 0;
 	ESQLitePreparedStatementStepResult StepResult = ESQLitePreparedStatementStepResult::Done;
 };
